Fix third party plugin template not having the correct include paths

diff --git a/Engine/Plugins/Editor/PluginBrowser/Templates/ThirdPartyLibrary/Source/PLUGIN_NAME/Private/PLUGIN_NAME.cpp b/Engine/Plugins/Editor/PluginBrowser/Templates/ThirdPartyLibrary/Source/PLUGIN_NAME/Private/PLUGIN_NAME.cpp
--- a/Engine/Plugins/Editor/PluginBrowser/Templates/ThirdPartyLibrary/Source/PLUGIN_NAME/Private/PLUGIN_NAME.cpp
+++ b/Engine/Plugins/Editor/PluginBrowser/Templates/ThirdPartyLibrary/Source/PLUGIN_NAME/Private/PLUGIN_NAME.cpp
@@ -4,7 +4,7 @@
 #include "Core.h"
 #include "Modules/ModuleManager.h"
 #include "Interfaces/IPluginManager.h"
-#include "ExampleLibrary.h"
+#include "PLUGIN_NAMELibrary/ExampleLibrary.h"
 
 #define LOCTEXT_NAMESPACE "FPLUGIN_NAMEModule"
 
