fix for ThirdParty Library plugin example failing to compile
added missing includes

diff --git a/Engine/Plugins/Editor/PluginBrowser/Templates/ThirdPartyLibrary/Source/PLUGIN_NAME/Private/PLUGIN_NAME.cpp b/Engine/Plugins/Editor/PluginBrowser/Templates/ThirdPartyLibrary/Source/PLUGIN_NAME/Private/PLUGIN_NAME.cpp
--- a/Engine/Plugins/Editor/PluginBrowser/Templates/ThirdPartyLibrary/Source/PLUGIN_NAME/Private/PLUGIN_NAME.cpp
+++ b/Engine/Plugins/Editor/PluginBrowser/Templates/ThirdPartyLibrary/Source/PLUGIN_NAME/Private/PLUGIN_NAME.cpp
@@ -4,6 +4,8 @@
 #include "Misc/MessageDialog.h"
 #include "Modules/ModuleManager.h"
 #include "Interfaces/IPluginManager.h"
+#include "Misc/Paths.h"
+#include "HAL/PlatformProcess.h"
 #include "PLUGIN_NAMELibrary/ExampleLibrary.h"
 
 #define LOCTEXT_NAMESPACE "FPLUGIN_NAMEModule"
